Fixed crash on Editor start with WarpUtils plugin enabled

diff --git a/Engine/Plugins/Runtime/WarpUtils/Shaders/Private/PFMExporterShaders.usf b/Engine/Plugins/Runtime/WarpUtils/Shaders/Private/PFMExporterShaders.usf
--- a/Engine/Plugins/Runtime/WarpUtils/Shaders/Private/PFMExporterShaders.usf
+++ b/Engine/Plugins/Runtime/WarpUtils/Shaders/Private/PFMExporterShaders.usf
@@ -1,5 +1,8 @@
 // Copyright Epic Games, Inc. All Rights Reserved.
 
+#include "/engine/public/platform.ush"
+
+
 float4x4	MeshToPFMMatrix;
 
 struct InputVS
